Add initial documentation for configuring the AutoSDK system. Currently just includes information for Windows.

#rnx

diff --git a/Engine/Extras/AutoSDK/HostWin64/Win64/README.md b/Engine/Extras/AutoSDK/HostWin64/Win64/README.md
new file mode 100644
--- /dev/null
+++ b/Engine/Extras/AutoSDK/HostWin64/Win64/README.md
@@ -0,0 +1,51 @@
+# Windows AutoSDK
+
+Visual Studio does not define any system-wide environment variables, so the Windows AutoSDK build environment is configured automatically by UnrealBuildTool based on the presence of the appropriately named directories. 
+
+The following components can be used with the AutoSDK system:
+
+## Visual C++
+
+Visual C++ toolchain files can be copied from the Tools/MSVC folder under the Visual Studio installation directory. For example:
+
+    C:\Program Files (x86)\Microsoft Visual Studio\2017\Professional\VC\Tools\MSVC\14.13.26128\...
+	
+Would be copied to:
+
+    [AutoSDK Root]\HostWin64\Win64\VS2017\14.13.26128\...
+
+Note that using each machine using a copy of the toolchain via AutoSDKs should have an appropriate license.
+
+## Windows SDK
+
+Multiple versions of the Windows SDK can be distributed via AutoSDK in the same merged directory hierarchy as multiple manual Windows SDK installations on a single machine. Typically:
+
+    C:\Program Files (x86)\Windows Kits\10\...
+
+Would be copied to:
+
+    [AutoSDK Root]\HostWin64\Win64\Windows Kits\10\...
+
+The NETFXSDK directory may also be distributed in the same manner.
+
+    [AutoSDK Root]\HostWin64\Win64\Windows Kits\NETFXSDK\...
+
+## DIA SDK
+
+The DIA SDK is distributed with Visual Studio, and is typically installed to the DIA SDK folder under the Visual Studio installation directory. For example:
+
+    C:\Program Files (x86)\Microsoft Visual Studio\2017\Professional\DIA SDK\...
+
+Would be copied to:
+
+    [AutoSDK Root]\HostWin64\Win64\DIA SDK\...
+
+## LLVM/Clang
+
+Multiple Clang versions can be distributed with AutoSDK, copied from the LLVM installation directory. For example:
+
+	C:\Program Files\LLVM\...
+	
+Would be copied to:
+
+    [AutoSDK Root]\HostWin64\Win64\LLVM\8.00\...
diff --git a/Engine/Extras/AutoSDK/README.md b/Engine/Extras/AutoSDK/README.md
new file mode 100644
--- /dev/null
+++ b/Engine/Extras/AutoSDK/README.md
@@ -0,0 +1,39 @@
+# AutoSDK
+
+## Introduction
+
+The AutoSDK system provides a mechanism for distributing target platform SDKs and configuring them for use by the engine on demand.
+
+It was designed for build machines to serve multiple branches with different SDK requirements without needing to manually manage installed packages, but may also be used by developers that don't require a full SDK install. Only a minimal toolset is typically available by default (compiler, deployment software, etc..).
+
+UnrealBuildTool, AutomationTool and the Unreal Editor are all designed to work seamlessly with AutoSDK. The switching between SDKs is handled by UnrealBuildTool, which is invoked by the other tools. 
+
+For any engine verison, UnrealBuildTool has a preferred SDK version that it will attempt to use. These versions are typically defined by classes derived from UEBuildPlatform in the UBT source code.
+
+## Setup
+
+Epic is not able to distribute SDKs for legal reasons, so this directory tree serves a template showing how to structure it. We recommend you start with an empty AutoSDK folder and copy files from here as you add a particular SDK version.
+
+**Do not configure the AutoSDK system to use this folder. It just a template, and the system will not work without adding additional files.**
+
+The AutoSDK folder should be submitted to source control separately to the game or engine code, so that multiple branches on a single machine can share it. 
+
+Any developer wishing to use AutoSDK can sync it and set the `UE_SDKS_ROOT` environment variable to point to the path on their local machine for UBT which contains it.
+
+## Layout
+
+The AutoSDK directory structure is loosely defined and generally left to the discretion of each platform, but generally follows this pattern.
+
+    /HostPlatform/                                 --- Host platform that the SDKs can be used on.
+    /HostPlatform/TargetPlatform/                  --- Directory containing the "TargetPlatform" SDKs for "HostPlatform" (eg. "/HostWin64/Android/").
+	/HostPlatform/TargetPlatform/1.0/              --- Directory containing the "TargetPlatform" 1.0 SDK for "HostPlatform".
+	/HostPlatform/TargetPlatform/1.0/Setup.bat     --- Optional batch file that is run when setting up this SDK. Named "setup.sh" on Mac/Linux.
+	/HostPlatform/TargetPlatform/1.0/Unsetup.bat   --- Optional batch file that is run when removing this SDK. Named "unsetup.sh" on Mac/Linux.
+
+When run, Setup.bat will output a text file to the same directory called OutputEnvVars.txt, containing a list of environment variables to set in the form NAME=VALUE, as well as special directives to modify the PATH environment variable such as ADDPATH=Foo and STRIPPATH=Foo.
+
+Setup scripts written by Epic for supported SDK versions are included under this directory.
+
+## Platforms
+
+More information about adding SDKs for each platform is given by README.md files in HostPlatform/TargetPlatform subfolders.
\ No newline at end of file
