Adding Lumin AutoSDK support files for licensees.

#rnx

diff --git a/Engine/Extras/AutoSDK/HostMac/Lumin/0.12/setup.sh b/Engine/Extras/AutoSDK/HostMac/Lumin/0.12/setup.sh
new file mode 100755
--- /dev/null
+++ b/Engine/Extras/AutoSDK/HostMac/Lumin/0.12/setup.sh
@@ -0,0 +1,5 @@
+#!/bin/bash
+
+cd `dirname $0`
+echo MLSDK=`pwd` >../OutputEnvVars.txt
+
diff --git a/Engine/Extras/AutoSDK/HostMac/Lumin/0.12/unsetup.sh b/Engine/Extras/AutoSDK/HostMac/Lumin/0.12/unsetup.sh
new file mode 100755
--- /dev/null
+++ b/Engine/Extras/AutoSDK/HostMac/Lumin/0.12/unsetup.sh
@@ -0,0 +1,4 @@
+#!/bin/bash
+
+cd `dirname $0`
+rm ../OutputEnvVars.txt 2>/dev/null || true
diff --git a/Engine/Extras/AutoSDK/HostMac/Lumin/0.16/setup.sh b/Engine/Extras/AutoSDK/HostMac/Lumin/0.16/setup.sh
new file mode 100755
--- /dev/null
+++ b/Engine/Extras/AutoSDK/HostMac/Lumin/0.16/setup.sh
@@ -0,0 +1,5 @@
+#!/bin/bash
+
+cd `dirname $0`
+echo MLSDK=`pwd` >../OutputEnvVars.txt
+
diff --git a/Engine/Extras/AutoSDK/HostMac/Lumin/0.16/unsetup.sh b/Engine/Extras/AutoSDK/HostMac/Lumin/0.16/unsetup.sh
new file mode 100755
--- /dev/null
+++ b/Engine/Extras/AutoSDK/HostMac/Lumin/0.16/unsetup.sh
@@ -0,0 +1,4 @@
+#!/bin/bash
+
+cd `dirname $0`
+rm ../OutputEnvVars.txt 2>/dev/null || true
diff --git a/Engine/Extras/AutoSDK/HostMac/Lumin/0.19/setup.sh b/Engine/Extras/AutoSDK/HostMac/Lumin/0.19/setup.sh
new file mode 100755
--- /dev/null
+++ b/Engine/Extras/AutoSDK/HostMac/Lumin/0.19/setup.sh
@@ -0,0 +1,5 @@
+#!/bin/bash
+
+cd `dirname $0`
+echo MLSDK=`pwd` >../OutputEnvVars.txt
+
diff --git a/Engine/Extras/AutoSDK/HostMac/Lumin/0.19/unsetup.sh b/Engine/Extras/AutoSDK/HostMac/Lumin/0.19/unsetup.sh
new file mode 100755
--- /dev/null
+++ b/Engine/Extras/AutoSDK/HostMac/Lumin/0.19/unsetup.sh
@@ -0,0 +1,4 @@
+#!/bin/bash
+
+cd `dirname $0`
+rm ../OutputEnvVars.txt 2>/dev/null || true
diff --git a/Engine/Extras/AutoSDK/HostMac/Lumin/README.md b/Engine/Extras/AutoSDK/HostMac/Lumin/README.md
new file mode 100644
--- /dev/null
+++ b/Engine/Extras/AutoSDK/HostMac/Lumin/README.md
@@ -0,0 +1,3 @@
+# Lumin AutoSDK
+
+Details forthcoming.
diff --git a/Engine/Extras/AutoSDK/HostWin64/Lumin/0.12/Setup.bat b/Engine/Extras/AutoSDK/HostWin64/Lumin/0.12/Setup.bat
new file mode 100644
--- /dev/null
+++ b/Engine/Extras/AutoSDK/HostWin64/Lumin/0.12/Setup.bat
@@ -0,0 +1,13 @@
+set CurrentFolder=%~dp0
+echo %CurrentFolder%
+
+REM delete OutputEnvVars from the old location
+del "%CurrentFolder%OutputEnvVars.txt"
+
+set OutputEnvVarsFolder=%~dp0..\
+del "%OutputEnvVarsFolder%OutputEnvVars.txt"
+
+
+
+REM set android sdk environment variables
+echo MLSDK=%CurrentFolder%>>"%OutputEnvVarsFolder%OutputEnvVars.txt"
diff --git a/Engine/Extras/AutoSDK/HostWin64/Lumin/0.12/Unsetup.bat b/Engine/Extras/AutoSDK/HostWin64/Lumin/0.12/Unsetup.bat
new file mode 100644
--- /dev/null
+++ b/Engine/Extras/AutoSDK/HostWin64/Lumin/0.12/Unsetup.bat
@@ -0,0 +1,4 @@
+rem TODO: check if this is under UE_SDKS_ROOT ?
+set CurrentFolder=%~dp0
+
+del %CurrentFolder%OutputEnvVars.txt
diff --git a/Engine/Extras/AutoSDK/HostWin64/Lumin/0.15/Setup.bat b/Engine/Extras/AutoSDK/HostWin64/Lumin/0.15/Setup.bat
new file mode 100644
--- /dev/null
+++ b/Engine/Extras/AutoSDK/HostWin64/Lumin/0.15/Setup.bat
@@ -0,0 +1,13 @@
+set CurrentFolder=%~dp0
+echo %CurrentFolder%
+
+REM delete OutputEnvVars from the old location
+del "%CurrentFolder%OutputEnvVars.txt"
+
+set OutputEnvVarsFolder=%~dp0..\
+del "%OutputEnvVarsFolder%OutputEnvVars.txt"
+
+
+
+REM set android sdk environment variables
+echo MLSDK=%CurrentFolder%>>"%OutputEnvVarsFolder%OutputEnvVars.txt"
diff --git a/Engine/Extras/AutoSDK/HostWin64/Lumin/0.15/Unsetup.bat b/Engine/Extras/AutoSDK/HostWin64/Lumin/0.15/Unsetup.bat
new file mode 100644
--- /dev/null
+++ b/Engine/Extras/AutoSDK/HostWin64/Lumin/0.15/Unsetup.bat
@@ -0,0 +1,4 @@
+rem TODO: check if this is under UE_SDKS_ROOT ?
+set CurrentFolder=%~dp0
+
+del %CurrentFolder%OutputEnvVars.txt
diff --git a/Engine/Extras/AutoSDK/HostWin64/Lumin/0.16/Setup.bat b/Engine/Extras/AutoSDK/HostWin64/Lumin/0.16/Setup.bat
new file mode 100644
--- /dev/null
+++ b/Engine/Extras/AutoSDK/HostWin64/Lumin/0.16/Setup.bat
@@ -0,0 +1,13 @@
+set CurrentFolder=%~dp0
+echo %CurrentFolder%
+
+REM delete OutputEnvVars from the old location
+del "%CurrentFolder%OutputEnvVars.txt"
+
+set OutputEnvVarsFolder=%~dp0..\
+del "%OutputEnvVarsFolder%OutputEnvVars.txt"
+
+
+
+REM set android sdk environment variables
+echo MLSDK=%CurrentFolder%>>"%OutputEnvVarsFolder%OutputEnvVars.txt"
diff --git a/Engine/Extras/AutoSDK/HostWin64/Lumin/0.16/Unsetup.bat b/Engine/Extras/AutoSDK/HostWin64/Lumin/0.16/Unsetup.bat
new file mode 100644
--- /dev/null
+++ b/Engine/Extras/AutoSDK/HostWin64/Lumin/0.16/Unsetup.bat
@@ -0,0 +1,4 @@
+rem TODO: check if this is under UE_SDKS_ROOT ?
+set CurrentFolder=%~dp0
+
+del %CurrentFolder%OutputEnvVars.txt
diff --git a/Engine/Extras/AutoSDK/HostWin64/Lumin/0.19/Setup.bat b/Engine/Extras/AutoSDK/HostWin64/Lumin/0.19/Setup.bat
new file mode 100644
--- /dev/null
+++ b/Engine/Extras/AutoSDK/HostWin64/Lumin/0.19/Setup.bat
@@ -0,0 +1,13 @@
+set CurrentFolder=%~dp0
+echo %CurrentFolder%
+
+REM delete OutputEnvVars from the old location
+del "%CurrentFolder%OutputEnvVars.txt"
+
+set OutputEnvVarsFolder=%~dp0..\
+del "%OutputEnvVarsFolder%OutputEnvVars.txt"
+
+
+
+REM set android sdk environment variables
+echo MLSDK=%CurrentFolder%>>"%OutputEnvVarsFolder%OutputEnvVars.txt"
diff --git a/Engine/Extras/AutoSDK/HostWin64/Lumin/0.19/Unsetup.bat b/Engine/Extras/AutoSDK/HostWin64/Lumin/0.19/Unsetup.bat
new file mode 100644
--- /dev/null
+++ b/Engine/Extras/AutoSDK/HostWin64/Lumin/0.19/Unsetup.bat
@@ -0,0 +1,4 @@
+rem TODO: check if this is under UE_SDKS_ROOT ?
+set CurrentFolder=%~dp0
+
+del %CurrentFolder%OutputEnvVars.txt
diff --git a/Engine/Extras/AutoSDK/HostWin64/Lumin/README.md b/Engine/Extras/AutoSDK/HostWin64/Lumin/README.md
new file mode 100644
--- /dev/null
+++ b/Engine/Extras/AutoSDK/HostWin64/Lumin/README.md
@@ -0,0 +1,3 @@
+# Lumin AutoSDK
+
+Details forthcoming.
