Adding Android AutoSDK support files for licensees.

#rnx

diff --git a/Engine/Extras/AutoSDK/HostWin64/Android/-19/Setup.bat b/Engine/Extras/AutoSDK/HostWin64/Android/-19/Setup.bat
new file mode 100644
--- /dev/null
+++ b/Engine/Extras/AutoSDK/HostWin64/Android/-19/Setup.bat
@@ -0,0 +1,34 @@
+rem TODO: check if this is under UE_SDKS_ROOT ?
+set CurrentFolder=%~dp0
+echo %CurrentFolder%
+
+REM delete OutputEnvVars from the old location
+del "%CurrentFolder%OutputEnvVars.txt"
+
+set OutputEnvVarsFolder=%~dp0..\
+del "%OutputEnvVarsFolder%OutputEnvVars.txt"
+
+
+
+REM set android sdk environment variables
+echo ANDROID_HOME=%CurrentFolder%android-sdk-windows>>"%OutputEnvVarsFolder%OutputEnvVars.txt"
+echo ANDROID_SDK_HOME=%CurrentFolder%android-sdk-windows>>"%OutputEnvVarsFolder%OutputEnvVars.txt"
+
+REM set ndk environment variables
+echo ANDROID_NDK_ROOT=%CurrentFolder%android-ndk-r9c>>"%OutputEnvVarsFolder%OutputEnvVars.txt"
+echo NDK_ROOT=%CurrentFolder%android-ndk-r9c>>"%OutputEnvVarsFolder%OutputEnvVars.txt"
+echo NDKROOT=%CurrentFolder%android-ndk-r9c>>"%OutputEnvVarsFolder%OutputEnvVars.txt"
+
+REM set ANT directory
+echo ANT_HOME=%CurrentFolder%apache-ant-1.8.2>>"%OutputEnvVarsFolder%OutputEnvVars.txt"
+
+
+REM set CYGWIN home directory
+echo CYGWIN_HOME=%CurrentFolder%cygwin>>"%OutputEnvVarsFolder%OutputEnvVars.txt"
+
+
+REM set JDK directory
+echo JAVA_HOME=%CurrentFolder%jdk1.6.0_45>>"%OutputEnvVarsFolder%OutputEnvVars.txt"
+
+REM support old branches that expect this file in a different location.
+copy "%OutputEnvVarsFolder%OutputEnvVars.txt" "%CurrentFolder%OutputEnvVars.txt"
diff --git a/Engine/Extras/AutoSDK/HostWin64/Android/-19/Unsetup.bat b/Engine/Extras/AutoSDK/HostWin64/Android/-19/Unsetup.bat
new file mode 100644
--- /dev/null
+++ b/Engine/Extras/AutoSDK/HostWin64/Android/-19/Unsetup.bat
@@ -0,0 +1,4 @@
+rem TODO: check if this is under UE_SDKS_ROOT ?
+set CurrentFolder=%~dp0
+
+del %CurrentFolder%OutputEnvVars.txt
diff --git a/Engine/Extras/AutoSDK/HostWin64/Android/-21/Setup.bat b/Engine/Extras/AutoSDK/HostWin64/Android/-21/Setup.bat
new file mode 100644
--- /dev/null
+++ b/Engine/Extras/AutoSDK/HostWin64/Android/-21/Setup.bat
@@ -0,0 +1,35 @@
+rem TODO: check if this is under UE_SDKS_ROOT ?
+set CurrentFolder=%~dp0
+echo %CurrentFolder%
+
+REM delete OutputEnvVars from the old location
+del "%CurrentFolder%OutputEnvVars.txt"
+
+set OutputEnvVarsFolder=%~dp0..\
+del "%OutputEnvVarsFolder%OutputEnvVars.txt"
+
+
+
+REM set android sdk environment variables
+echo ANDROID_HOME=%CurrentFolder%android-sdk-windows>>"%OutputEnvVarsFolder%OutputEnvVars.txt"
+echo ANDROID_SDK_HOME=%CurrentFolder%android-sdk-windows>>"%OutputEnvVarsFolder%OutputEnvVars.txt"
+
+REM set ndk environment variables
+echo ANDROID_NDK_ROOT=%CurrentFolder%android-ndk-r12b>>"%OutputEnvVarsFolder%OutputEnvVars.txt"
+echo NDK_ROOT=%CurrentFolder%android-ndk-r12b>>"%OutputEnvVarsFolder%OutputEnvVars.txt"
+echo NDKROOT=%CurrentFolder%android-ndk-r12b>>"%OutputEnvVarsFolder%OutputEnvVars.txt"
+
+REM set ANT directory
+echo ANT_HOME=%CurrentFolder%apache-ant-1.8.2>>"%OutputEnvVarsFolder%OutputEnvVars.txt"
+
+REM set JDK directory
+echo JAVA_HOME=%CurrentFolder%jdk1.8.0_77>>"%OutputEnvVarsFolder%OutputEnvVars.txt"
+
+REM Set the Andriod SWT path 
+echo ANDROID_SWT=%CurrentFolder%android-sdk-windows\tools\lib\x86_64>>"%OutputEnvVarsFolder%OutputEnvVars.txt"
+
+
+
+
+REM support old branches that expect this file in a different location.
+copy "%OutputEnvVarsFolder%OutputEnvVars.txt" "%CurrentFolder%OutputEnvVars.txt"
diff --git a/Engine/Extras/AutoSDK/HostWin64/Android/-21/Unsetup.bat b/Engine/Extras/AutoSDK/HostWin64/Android/-21/Unsetup.bat
new file mode 100644
--- /dev/null
+++ b/Engine/Extras/AutoSDK/HostWin64/Android/-21/Unsetup.bat
@@ -0,0 +1,4 @@
+rem TODO: check if this is under UE_SDKS_ROOT ?
+set CurrentFolder=%~dp0
+
+del %CurrentFolder%OutputEnvVars.txt
diff --git a/Engine/Extras/AutoSDK/HostWin64/Android/-22/Setup.bat b/Engine/Extras/AutoSDK/HostWin64/Android/-22/Setup.bat
new file mode 100644
--- /dev/null
+++ b/Engine/Extras/AutoSDK/HostWin64/Android/-22/Setup.bat
@@ -0,0 +1,36 @@
+rem TODO: check if this is under UE_SDKS_ROOT ?
+set CurrentFolder=%~dp0
+echo %CurrentFolder%
+
+REM delete OutputEnvVars from the old location
+del "%CurrentFolder%OutputEnvVars.txt"
+
+set OutputEnvVarsFolder=%~dp0..\
+del "%OutputEnvVarsFolder%OutputEnvVars.txt"
+
+
+
+REM set android sdk environment variables
+echo ANDROID_HOME=%CurrentFolder%android-sdk-windows>>"%OutputEnvVarsFolder%OutputEnvVars.txt"
+echo ANDROID_SDK_HOME=%CurrentFolder%android-sdk-windows>>"%OutputEnvVarsFolder%OutputEnvVars.txt"
+
+REM set ndk environment variables
+echo ANDROID_NDK_ROOT=%CurrentFolder%android-ndk-r14b>>"%OutputEnvVarsFolder%OutputEnvVars.txt"
+echo NDK_ROOT=%CurrentFolder%android-ndk-r14b>>"%OutputEnvVarsFolder%OutputEnvVars.txt"
+echo NDKROOT=%CurrentFolder%android-ndk-r14b>>"%OutputEnvVarsFolder%OutputEnvVars.txt"
+
+REM set ANT directory
+echo ANT_HOME=%CurrentFolder%apache-ant-1.8.2>>"%OutputEnvVarsFolder%OutputEnvVars.txt"
+
+REM set JDK directory
+echo JAVA_HOME=%CurrentFolder%jdk1.8.0_77>>"%OutputEnvVarsFolder%OutputEnvVars.txt"
+
+REM Set the Andriod SWT path 
+echo ANDROID_SWT=%CurrentFolder%android-sdk-windows\tools\lib\x86_64>>"%OutputEnvVarsFolder%OutputEnvVars.txt"
+
+
+
+
+REM support old branches that expect this file in a different location.
+copy "%OutputEnvVarsFolder%OutputEnvVars.txt" "%CurrentFolder%OutputEnvVars.txt"
+
diff --git a/Engine/Extras/AutoSDK/HostWin64/Android/-22/Unsetup.bat b/Engine/Extras/AutoSDK/HostWin64/Android/-22/Unsetup.bat
new file mode 100644
--- /dev/null
+++ b/Engine/Extras/AutoSDK/HostWin64/Android/-22/Unsetup.bat
@@ -0,0 +1,4 @@
+rem TODO: check if this is under UE_SDKS_ROOT ?
+set CurrentFolder=%~dp0
+
+del %CurrentFolder%OutputEnvVars.txt
diff --git a/Engine/Extras/AutoSDK/HostWin64/Android/README.md b/Engine/Extras/AutoSDK/HostWin64/Android/README.md
new file mode 100644
--- /dev/null
+++ b/Engine/Extras/AutoSDK/HostWin64/Android/README.md
@@ -0,0 +1,3 @@
+# Android AutoSDK
+
+Details forthcoming.
diff --git a/Engine/Extras/AutoSDK/HostWin64/Android/SetupEnvironmentVars.bat b/Engine/Extras/AutoSDK/HostWin64/Android/SetupEnvironmentVars.bat
new file mode 100644
--- /dev/null
+++ b/Engine/Extras/AutoSDK/HostWin64/Android/SetupEnvironmentVars.bat
@@ -0,0 +1 @@
+if not "%UE_SDKS_ROOT%"=="" (for /F "tokens=*" %%A in (%UE_SDKS_ROOT%\HostWin64\Android\OutputEnvVars.txt) do set %%A)
\ No newline at end of file
