Adding Linux AutoSDK support files for licensees.

diff --git a/Engine/Extras/AutoSDK/HostLinux/Linux_x64/README.md b/Engine/Extras/AutoSDK/HostLinux/Linux_x64/README.md
new file mode 100644
--- /dev/null
+++ b/Engine/Extras/AutoSDK/HostLinux/Linux_x64/README.md
@@ -0,0 +1,3 @@
+# Linux AutoSDK
+
+Details forthcoming.
diff --git a/Engine/Extras/AutoSDK/HostLinux/Linux_x64/v11_clang-5.0.0-centos7/setup.sh b/Engine/Extras/AutoSDK/HostLinux/Linux_x64/v11_clang-5.0.0-centos7/setup.sh
new file mode 100755
--- /dev/null
+++ b/Engine/Extras/AutoSDK/HostLinux/Linux_x64/v11_clang-5.0.0-centos7/setup.sh
@@ -0,0 +1,14 @@
+#!/bin/bash
+
+# TODO: check if this is under UE_SDKS_ROOT ?
+CURRENT_FOLDER=$(cd "$(dirname "$BASH_SOURCE")" ; pwd)
+echo "$CURRENT_FOLDER"
+
+OUTPUT_ENV_VARS_FOLDER="$CURRENT_FOLDER/.."
+rm "$OUTPUT_ENV_VARS_FOLDER/OutputEnvVars.txt" 2> /dev/null
+
+# set new LINUX_MULTIARCH_ROOT that supersedes it
+echo LINUX_MULTIARCH_ROOT="$CURRENT_FOLDER" > "$OUTPUT_ENV_VARS_FOLDER/OutputEnvVars.txt"
+
+# support old branches that expect this file in a different location.
+cp "$OUTPUT_ENV_VARS_FOLDER/OutputEnvVars.txt" "$CURRENT_FOLDER"
diff --git a/Engine/Extras/AutoSDK/HostLinux/Linux_x64/v11_clang-5.0.0-centos7/unsetup.sh b/Engine/Extras/AutoSDK/HostLinux/Linux_x64/v11_clang-5.0.0-centos7/unsetup.sh
new file mode 100755
--- /dev/null
+++ b/Engine/Extras/AutoSDK/HostLinux/Linux_x64/v11_clang-5.0.0-centos7/unsetup.sh
@@ -0,0 +1,9 @@
+#!/bin/bash
+
+# TODO: check if this is under UE_SDKS_ROOT ?
+CURRENT_FOLDER=$(cd "$(dirname "$BASH_SOURCE")" ; pwd)
+echo "$CURRENT_FOLDER"
+
+OUTPUT_ENV_VARS_FOLDER="$CURRENT_FOLDER/.."
+
+rm "$OUTPUT_ENV_VARS_FOLDER/OutputEnvVars.txt"
diff --git a/Engine/Extras/AutoSDK/HostLinux/Linux_x64/v12_clang-6.0.1-centos7/setup.sh b/Engine/Extras/AutoSDK/HostLinux/Linux_x64/v12_clang-6.0.1-centos7/setup.sh
new file mode 100755
--- /dev/null
+++ b/Engine/Extras/AutoSDK/HostLinux/Linux_x64/v12_clang-6.0.1-centos7/setup.sh
@@ -0,0 +1,14 @@
+#!/bin/bash
+
+# TODO: check if this is under UE_SDKS_ROOT ?
+CURRENT_FOLDER=$(cd "$(dirname "$BASH_SOURCE")" ; pwd)
+echo "$CURRENT_FOLDER"
+
+OUTPUT_ENV_VARS_FOLDER="$CURRENT_FOLDER/.."
+rm "$OUTPUT_ENV_VARS_FOLDER/OutputEnvVars.txt" 2> /dev/null
+
+# set new LINUX_MULTIARCH_ROOT that supersedes it
+echo LINUX_MULTIARCH_ROOT="$CURRENT_FOLDER" > "$OUTPUT_ENV_VARS_FOLDER/OutputEnvVars.txt"
+
+# support old branches that expect this file in a different location.
+cp "$OUTPUT_ENV_VARS_FOLDER/OutputEnvVars.txt" "$CURRENT_FOLDER"
diff --git a/Engine/Extras/AutoSDK/HostLinux/Linux_x64/v12_clang-6.0.1-centos7/unsetup.sh b/Engine/Extras/AutoSDK/HostLinux/Linux_x64/v12_clang-6.0.1-centos7/unsetup.sh
new file mode 100755
--- /dev/null
+++ b/Engine/Extras/AutoSDK/HostLinux/Linux_x64/v12_clang-6.0.1-centos7/unsetup.sh
@@ -0,0 +1,9 @@
+#!/bin/bash
+
+# TODO: check if this is under UE_SDKS_ROOT ?
+CURRENT_FOLDER=$(cd "$(dirname "$BASH_SOURCE")" ; pwd)
+echo "$CURRENT_FOLDER"
+
+OUTPUT_ENV_VARS_FOLDER="$CURRENT_FOLDER/.."
+
+rm "$OUTPUT_ENV_VARS_FOLDER/OutputEnvVars.txt"
diff --git a/Engine/Extras/AutoSDK/HostLinux/Linux_x64/v13_clang-7.0.1-centos7/setup.sh b/Engine/Extras/AutoSDK/HostLinux/Linux_x64/v13_clang-7.0.1-centos7/setup.sh
new file mode 100755
--- /dev/null
+++ b/Engine/Extras/AutoSDK/HostLinux/Linux_x64/v13_clang-7.0.1-centos7/setup.sh
@@ -0,0 +1,14 @@
+#!/bin/bash
+
+# TODO: check if this is under UE_SDKS_ROOT ?
+CURRENT_FOLDER=$(cd "$(dirname "$BASH_SOURCE")" ; pwd)
+echo "$CURRENT_FOLDER"
+
+OUTPUT_ENV_VARS_FOLDER="$CURRENT_FOLDER/.."
+rm "$OUTPUT_ENV_VARS_FOLDER/OutputEnvVars.txt" 2> /dev/null
+
+# set new LINUX_MULTIARCH_ROOT that supersedes it
+echo LINUX_MULTIARCH_ROOT="$CURRENT_FOLDER" > "$OUTPUT_ENV_VARS_FOLDER/OutputEnvVars.txt"
+
+# support old branches that expect this file in a different location.
+cp "$OUTPUT_ENV_VARS_FOLDER/OutputEnvVars.txt" "$CURRENT_FOLDER"
diff --git a/Engine/Extras/AutoSDK/HostLinux/Linux_x64/v13_clang-7.0.1-centos7/unsetup.sh b/Engine/Extras/AutoSDK/HostLinux/Linux_x64/v13_clang-7.0.1-centos7/unsetup.sh
new file mode 100755
--- /dev/null
+++ b/Engine/Extras/AutoSDK/HostLinux/Linux_x64/v13_clang-7.0.1-centos7/unsetup.sh
@@ -0,0 +1,9 @@
+#!/bin/bash
+
+# TODO: check if this is under UE_SDKS_ROOT ?
+CURRENT_FOLDER=$(cd "$(dirname "$BASH_SOURCE")" ; pwd)
+echo "$CURRENT_FOLDER"
+
+OUTPUT_ENV_VARS_FOLDER="$CURRENT_FOLDER/.."
+
+rm "$OUTPUT_ENV_VARS_FOLDER/OutputEnvVars.txt"
diff --git a/Engine/Extras/AutoSDK/HostLinux/Linux_x64/v14_clang-8.0.1-centos7/setup.sh b/Engine/Extras/AutoSDK/HostLinux/Linux_x64/v14_clang-8.0.1-centos7/setup.sh
new file mode 100755
--- /dev/null
+++ b/Engine/Extras/AutoSDK/HostLinux/Linux_x64/v14_clang-8.0.1-centos7/setup.sh
@@ -0,0 +1,14 @@
+#!/bin/bash
+
+# TODO: check if this is under UE_SDKS_ROOT ?
+CURRENT_FOLDER=$(cd "$(dirname "$BASH_SOURCE")" ; pwd)
+echo "$CURRENT_FOLDER"
+
+OUTPUT_ENV_VARS_FOLDER="$CURRENT_FOLDER/.."
+rm "$OUTPUT_ENV_VARS_FOLDER/OutputEnvVars.txt" 2> /dev/null
+
+# set new LINUX_MULTIARCH_ROOT that supersedes it
+echo LINUX_MULTIARCH_ROOT="$CURRENT_FOLDER" > "$OUTPUT_ENV_VARS_FOLDER/OutputEnvVars.txt"
+
+# support old branches that expect this file in a different location.
+cp "$OUTPUT_ENV_VARS_FOLDER/OutputEnvVars.txt" "$CURRENT_FOLDER"
diff --git a/Engine/Extras/AutoSDK/HostLinux/Linux_x64/v14_clang-8.0.1-centos7/unsetup.sh b/Engine/Extras/AutoSDK/HostLinux/Linux_x64/v14_clang-8.0.1-centos7/unsetup.sh
new file mode 100755
--- /dev/null
+++ b/Engine/Extras/AutoSDK/HostLinux/Linux_x64/v14_clang-8.0.1-centos7/unsetup.sh
@@ -0,0 +1,9 @@
+#!/bin/bash
+
+# TODO: check if this is under UE_SDKS_ROOT ?
+CURRENT_FOLDER=$(cd "$(dirname "$BASH_SOURCE")" ; pwd)
+echo "$CURRENT_FOLDER"
+
+OUTPUT_ENV_VARS_FOLDER="$CURRENT_FOLDER/.."
+
+rm "$OUTPUT_ENV_VARS_FOLDER/OutputEnvVars.txt"
diff --git a/Engine/Extras/AutoSDK/HostWin64/Linux_x64/README.md b/Engine/Extras/AutoSDK/HostWin64/Linux_x64/README.md
new file mode 100644
--- /dev/null
+++ b/Engine/Extras/AutoSDK/HostWin64/Linux_x64/README.md
@@ -0,0 +1,3 @@
+# Linux AutoSDK
+
+Details forthcoming.
diff --git a/Engine/Extras/AutoSDK/HostWin64/Linux_x64/arm-unknown-linux-gnueabihf_v5_clang-3.5.0-ld-2.23.1-glibc-2.13/Setup.bat b/Engine/Extras/AutoSDK/HostWin64/Linux_x64/arm-unknown-linux-gnueabihf_v5_clang-3.5.0-ld-2.23.1-glibc-2.13/Setup.bat
new file mode 100644
--- /dev/null
+++ b/Engine/Extras/AutoSDK/HostWin64/Linux_x64/arm-unknown-linux-gnueabihf_v5_clang-3.5.0-ld-2.23.1-glibc-2.13/Setup.bat
@@ -0,0 +1,14 @@
+rem TODO: check if this is under UE_SDKS_ROOT ?
+set CurrentFolder=%~dp0
+echo %CurrentFolder%
+
+REM delete OutputEnvVars from the old location
+del "%CurrentFolder%OutputEnvVars.txt"
+
+set OutputEnvVarsFolder=%~dp0..\
+del "%OutputEnvVarsFolder%OutputEnvVars.txt"
+
+echo LINUX_ROOT=%CurrentFolder%toolchain>"%OutputEnvVarsFolder%OutputEnvVars.txt"
+
+REM support old branches that expect this file in a different location.
+copy "%OutputEnvVarsFolder%OutputEnvVars.txt" "%CurrentFolder%OutputEnvVars.txt"
diff --git a/Engine/Extras/AutoSDK/HostWin64/Linux_x64/arm-unknown-linux-gnueabihf_v5_clang-3.5.0-ld-2.23.1-glibc-2.13/Unsetup.bat b/Engine/Extras/AutoSDK/HostWin64/Linux_x64/arm-unknown-linux-gnueabihf_v5_clang-3.5.0-ld-2.23.1-glibc-2.13/Unsetup.bat
new file mode 100644
--- /dev/null
+++ b/Engine/Extras/AutoSDK/HostWin64/Linux_x64/arm-unknown-linux-gnueabihf_v5_clang-3.5.0-ld-2.23.1-glibc-2.13/Unsetup.bat
@@ -0,0 +1,4 @@
+rem TODO: check if this is under UE_SDKS_ROOT ?
+set CurrentFolder=%~dp0
+
+del %CurrentFolder%OutputEnvVars.txt
diff --git a/Engine/Extras/AutoSDK/HostWin64/Linux_x64/v10_clang-5.0.0-centos7/Setup.bat b/Engine/Extras/AutoSDK/HostWin64/Linux_x64/v10_clang-5.0.0-centos7/Setup.bat
new file mode 100644
--- /dev/null
+++ b/Engine/Extras/AutoSDK/HostWin64/Linux_x64/v10_clang-5.0.0-centos7/Setup.bat
@@ -0,0 +1,18 @@
+rem TODO: check if this is under UE_SDKS_ROOT ?
+set CurrentFolder=%~dp0
+echo %CurrentFolder%
+
+REM delete OutputEnvVars from the old location
+del "%CurrentFolder%OutputEnvVars.txt"
+
+set OutputEnvVarsFolder=%~dp0..\
+del "%OutputEnvVarsFolder%OutputEnvVars.txt"
+
+rem Still set LINUX_ROOT for compatibility
+echo LINUX_ROOT=%CurrentFolder%x86_64-unknown-linux-gnu>"%OutputEnvVarsFolder%OutputEnvVars.txt"
+
+rem set new LINUX_MULTIARCH_ROOT that supersedes it
+echo LINUX_MULTIARCH_ROOT=%CurrentFolder%>>"%OutputEnvVarsFolder%OutputEnvVars.txt"
+
+REM support old branches that expect this file in a different location.
+copy "%OutputEnvVarsFolder%OutputEnvVars.txt" "%CurrentFolder%OutputEnvVars.txt"
diff --git a/Engine/Extras/AutoSDK/HostWin64/Linux_x64/v10_clang-5.0.0-centos7/Unsetup.bat b/Engine/Extras/AutoSDK/HostWin64/Linux_x64/v10_clang-5.0.0-centos7/Unsetup.bat
new file mode 100644
--- /dev/null
+++ b/Engine/Extras/AutoSDK/HostWin64/Linux_x64/v10_clang-5.0.0-centos7/Unsetup.bat
@@ -0,0 +1,4 @@
+rem TODO: check if this is under UE_SDKS_ROOT ?
+set CurrentFolder=%~dp0
+
+del %CurrentFolder%OutputEnvVars.txt
diff --git a/Engine/Extras/AutoSDK/HostWin64/Linux_x64/v11_clang-5.0.0-centos7/Setup.bat b/Engine/Extras/AutoSDK/HostWin64/Linux_x64/v11_clang-5.0.0-centos7/Setup.bat
new file mode 100644
--- /dev/null
+++ b/Engine/Extras/AutoSDK/HostWin64/Linux_x64/v11_clang-5.0.0-centos7/Setup.bat
@@ -0,0 +1,18 @@
+rem TODO: check if this is under UE_SDKS_ROOT ?
+set CurrentFolder=%~dp0
+echo %CurrentFolder%
+
+REM delete OutputEnvVars from the old location
+del "%CurrentFolder%OutputEnvVars.txt"
+
+set OutputEnvVarsFolder=%~dp0..\
+del "%OutputEnvVarsFolder%OutputEnvVars.txt"
+
+rem Still set LINUX_ROOT for compatibility
+echo LINUX_ROOT=%CurrentFolder%x86_64-unknown-linux-gnu>"%OutputEnvVarsFolder%OutputEnvVars.txt"
+
+rem set new LINUX_MULTIARCH_ROOT that supersedes it
+echo LINUX_MULTIARCH_ROOT=%CurrentFolder%>>"%OutputEnvVarsFolder%OutputEnvVars.txt"
+
+REM support old branches that expect this file in a different location.
+copy "%OutputEnvVarsFolder%OutputEnvVars.txt" "%CurrentFolder%OutputEnvVars.txt"
diff --git a/Engine/Extras/AutoSDK/HostWin64/Linux_x64/v11_clang-5.0.0-centos7/Unsetup.bat b/Engine/Extras/AutoSDK/HostWin64/Linux_x64/v11_clang-5.0.0-centos7/Unsetup.bat
new file mode 100644
--- /dev/null
+++ b/Engine/Extras/AutoSDK/HostWin64/Linux_x64/v11_clang-5.0.0-centos7/Unsetup.bat
@@ -0,0 +1,4 @@
+rem TODO: check if this is under UE_SDKS_ROOT ?
+set CurrentFolder=%~dp0
+
+del %CurrentFolder%OutputEnvVars.txt
diff --git a/Engine/Extras/AutoSDK/HostWin64/Linux_x64/v12_clang-6.0.1-centos7/Setup.bat b/Engine/Extras/AutoSDK/HostWin64/Linux_x64/v12_clang-6.0.1-centos7/Setup.bat
new file mode 100644
--- /dev/null
+++ b/Engine/Extras/AutoSDK/HostWin64/Linux_x64/v12_clang-6.0.1-centos7/Setup.bat
@@ -0,0 +1,18 @@
+rem TODO: check if this is under UE_SDKS_ROOT ?
+set CurrentFolder=%~dp0
+echo %CurrentFolder%
+
+REM delete OutputEnvVars from the old location
+del "%CurrentFolder%OutputEnvVars.txt"
+
+set OutputEnvVarsFolder=%~dp0..\
+del "%OutputEnvVarsFolder%OutputEnvVars.txt"
+
+rem Still set LINUX_ROOT for compatibility
+echo LINUX_ROOT=%CurrentFolder%x86_64-unknown-linux-gnu>"%OutputEnvVarsFolder%OutputEnvVars.txt"
+
+rem set new LINUX_MULTIARCH_ROOT that supersedes it
+echo LINUX_MULTIARCH_ROOT=%CurrentFolder%>>"%OutputEnvVarsFolder%OutputEnvVars.txt"
+
+REM support old branches that expect this file in a different location.
+copy "%OutputEnvVarsFolder%OutputEnvVars.txt" "%CurrentFolder%OutputEnvVars.txt"
diff --git a/Engine/Extras/AutoSDK/HostWin64/Linux_x64/v12_clang-6.0.1-centos7/Unsetup.bat b/Engine/Extras/AutoSDK/HostWin64/Linux_x64/v12_clang-6.0.1-centos7/Unsetup.bat
new file mode 100644
--- /dev/null
+++ b/Engine/Extras/AutoSDK/HostWin64/Linux_x64/v12_clang-6.0.1-centos7/Unsetup.bat
@@ -0,0 +1,4 @@
+rem TODO: check if this is under UE_SDKS_ROOT ?
+set CurrentFolder=%~dp0
+
+del %CurrentFolder%OutputEnvVars.txt
diff --git a/Engine/Extras/AutoSDK/HostWin64/Linux_x64/v13_clang-7.0.1-centos7/Setup.bat b/Engine/Extras/AutoSDK/HostWin64/Linux_x64/v13_clang-7.0.1-centos7/Setup.bat
new file mode 100644
--- /dev/null
+++ b/Engine/Extras/AutoSDK/HostWin64/Linux_x64/v13_clang-7.0.1-centos7/Setup.bat
@@ -0,0 +1,18 @@
+rem TODO: check if this is under UE_SDKS_ROOT ?
+set CurrentFolder=%~dp0
+echo %CurrentFolder%
+
+REM delete OutputEnvVars from the old location
+del "%CurrentFolder%OutputEnvVars.txt"
+
+set OutputEnvVarsFolder=%~dp0..\
+del "%OutputEnvVarsFolder%OutputEnvVars.txt"
+
+rem Still set LINUX_ROOT for compatibility
+echo LINUX_ROOT=%CurrentFolder%x86_64-unknown-linux-gnu>"%OutputEnvVarsFolder%OutputEnvVars.txt"
+
+rem set new LINUX_MULTIARCH_ROOT that supersedes it
+echo LINUX_MULTIARCH_ROOT=%CurrentFolder%>>"%OutputEnvVarsFolder%OutputEnvVars.txt"
+
+REM support old branches that expect this file in a different location.
+copy "%OutputEnvVarsFolder%OutputEnvVars.txt" "%CurrentFolder%OutputEnvVars.txt"
diff --git a/Engine/Extras/AutoSDK/HostWin64/Linux_x64/v13_clang-7.0.1-centos7/Unsetup.bat b/Engine/Extras/AutoSDK/HostWin64/Linux_x64/v13_clang-7.0.1-centos7/Unsetup.bat
new file mode 100644
--- /dev/null
+++ b/Engine/Extras/AutoSDK/HostWin64/Linux_x64/v13_clang-7.0.1-centos7/Unsetup.bat
@@ -0,0 +1,4 @@
+rem TODO: check if this is under UE_SDKS_ROOT ?
+set CurrentFolder=%~dp0
+
+del %CurrentFolder%OutputEnvVars.txt
diff --git a/Engine/Extras/AutoSDK/HostWin64/Linux_x64/v14_clang-8.0.1-centos7/Setup.bat b/Engine/Extras/AutoSDK/HostWin64/Linux_x64/v14_clang-8.0.1-centos7/Setup.bat
new file mode 100644
--- /dev/null
+++ b/Engine/Extras/AutoSDK/HostWin64/Linux_x64/v14_clang-8.0.1-centos7/Setup.bat
@@ -0,0 +1,18 @@
+rem TODO: check if this is under UE_SDKS_ROOT ?
+set CurrentFolder=%~dp0
+echo %CurrentFolder%
+
+REM delete OutputEnvVars from the old location
+del "%CurrentFolder%OutputEnvVars.txt"
+
+set OutputEnvVarsFolder=%~dp0..\
+del "%OutputEnvVarsFolder%OutputEnvVars.txt"
+
+rem Still set LINUX_ROOT for compatibility
+echo LINUX_ROOT=%CurrentFolder%x86_64-unknown-linux-gnu>"%OutputEnvVarsFolder%OutputEnvVars.txt"
+
+rem set new LINUX_MULTIARCH_ROOT that supersedes it
+echo LINUX_MULTIARCH_ROOT=%CurrentFolder%>>"%OutputEnvVarsFolder%OutputEnvVars.txt"
+
+REM support old branches that expect this file in a different location.
+copy "%OutputEnvVarsFolder%OutputEnvVars.txt" "%CurrentFolder%OutputEnvVars.txt"
diff --git a/Engine/Extras/AutoSDK/HostWin64/Linux_x64/v14_clang-8.0.1-centos7/Unsetup.bat b/Engine/Extras/AutoSDK/HostWin64/Linux_x64/v14_clang-8.0.1-centos7/Unsetup.bat
new file mode 100644
--- /dev/null
+++ b/Engine/Extras/AutoSDK/HostWin64/Linux_x64/v14_clang-8.0.1-centos7/Unsetup.bat
@@ -0,0 +1,4 @@
+rem TODO: check if this is under UE_SDKS_ROOT ?
+set CurrentFolder=%~dp0
+
+del %CurrentFolder%OutputEnvVars.txt
diff --git a/Engine/Extras/AutoSDK/HostWin64/Linux_x64/v3_clang-3.3_ld-2.24_glibc-2.12.2/Setup.bat b/Engine/Extras/AutoSDK/HostWin64/Linux_x64/v3_clang-3.3_ld-2.24_glibc-2.12.2/Setup.bat
new file mode 100644
--- /dev/null
+++ b/Engine/Extras/AutoSDK/HostWin64/Linux_x64/v3_clang-3.3_ld-2.24_glibc-2.12.2/Setup.bat
@@ -0,0 +1,14 @@
+rem TODO: check if this is under UE_SDKS_ROOT ?
+set CurrentFolder=%~dp0
+echo %CurrentFolder%
+
+REM delete OutputEnvVars from the old location
+del "%CurrentFolder%OutputEnvVars.txt"
+
+set OutputEnvVarsFolder=%~dp0..\
+del "%OutputEnvVarsFolder%OutputEnvVars.txt"
+
+echo LINUX_ROOT=%CurrentFolder%toolchain>"%OutputEnvVarsFolder%OutputEnvVars.txt"
+
+REM support old branches that expect this file in a different location.
+copy "%OutputEnvVarsFolder%OutputEnvVars.txt" "%CurrentFolder%OutputEnvVars.txt"
diff --git a/Engine/Extras/AutoSDK/HostWin64/Linux_x64/v3_clang-3.3_ld-2.24_glibc-2.12.2/Unsetup.bat b/Engine/Extras/AutoSDK/HostWin64/Linux_x64/v3_clang-3.3_ld-2.24_glibc-2.12.2/Unsetup.bat
new file mode 100644
--- /dev/null
+++ b/Engine/Extras/AutoSDK/HostWin64/Linux_x64/v3_clang-3.3_ld-2.24_glibc-2.12.2/Unsetup.bat
@@ -0,0 +1,4 @@
+rem TODO: check if this is under UE_SDKS_ROOT ?
+set CurrentFolder=%~dp0
+
+del %CurrentFolder%OutputEnvVars.txt
diff --git a/Engine/Extras/AutoSDK/HostWin64/Linux_x64/v4_clang-3.5.0_ld-2.24_glibc-2.12.2/Setup.bat b/Engine/Extras/AutoSDK/HostWin64/Linux_x64/v4_clang-3.5.0_ld-2.24_glibc-2.12.2/Setup.bat
new file mode 100644
--- /dev/null
+++ b/Engine/Extras/AutoSDK/HostWin64/Linux_x64/v4_clang-3.5.0_ld-2.24_glibc-2.12.2/Setup.bat
@@ -0,0 +1,14 @@
+rem TODO: check if this is under UE_SDKS_ROOT ?
+set CurrentFolder=%~dp0
+echo %CurrentFolder%
+
+REM delete OutputEnvVars from the old location
+del "%CurrentFolder%OutputEnvVars.txt"
+
+set OutputEnvVarsFolder=%~dp0..\
+del "%OutputEnvVarsFolder%OutputEnvVars.txt"
+
+echo LINUX_ROOT=%CurrentFolder%toolchain>"%OutputEnvVarsFolder%OutputEnvVars.txt"
+
+REM support old branches that expect this file in a different location.
+copy "%OutputEnvVarsFolder%OutputEnvVars.txt" "%CurrentFolder%OutputEnvVars.txt"
diff --git a/Engine/Extras/AutoSDK/HostWin64/Linux_x64/v4_clang-3.5.0_ld-2.24_glibc-2.12.2/Unsetup.bat b/Engine/Extras/AutoSDK/HostWin64/Linux_x64/v4_clang-3.5.0_ld-2.24_glibc-2.12.2/Unsetup.bat
new file mode 100644
--- /dev/null
+++ b/Engine/Extras/AutoSDK/HostWin64/Linux_x64/v4_clang-3.5.0_ld-2.24_glibc-2.12.2/Unsetup.bat
@@ -0,0 +1,4 @@
+rem TODO: check if this is under UE_SDKS_ROOT ?
+set CurrentFolder=%~dp0
+
+del %CurrentFolder%OutputEnvVars.txt
diff --git a/Engine/Extras/AutoSDK/HostWin64/Linux_x64/v6_clang-3.6.0_ld-2.24_glibc-2.12.2/Setup.bat b/Engine/Extras/AutoSDK/HostWin64/Linux_x64/v6_clang-3.6.0_ld-2.24_glibc-2.12.2/Setup.bat
new file mode 100644
--- /dev/null
+++ b/Engine/Extras/AutoSDK/HostWin64/Linux_x64/v6_clang-3.6.0_ld-2.24_glibc-2.12.2/Setup.bat
@@ -0,0 +1,14 @@
+rem TODO: check if this is under UE_SDKS_ROOT ?
+set CurrentFolder=%~dp0
+echo %CurrentFolder%
+
+REM delete OutputEnvVars from the old location
+del "%CurrentFolder%OutputEnvVars.txt"
+
+set OutputEnvVarsFolder=%~dp0..\
+del "%OutputEnvVarsFolder%OutputEnvVars.txt"
+
+echo LINUX_ROOT=%CurrentFolder%toolchain>"%OutputEnvVarsFolder%OutputEnvVars.txt"
+
+REM support old branches that expect this file in a different location.
+copy "%OutputEnvVarsFolder%OutputEnvVars.txt" "%CurrentFolder%OutputEnvVars.txt"
diff --git a/Engine/Extras/AutoSDK/HostWin64/Linux_x64/v6_clang-3.6.0_ld-2.24_glibc-2.12.2/Unsetup.bat b/Engine/Extras/AutoSDK/HostWin64/Linux_x64/v6_clang-3.6.0_ld-2.24_glibc-2.12.2/Unsetup.bat
new file mode 100644
--- /dev/null
+++ b/Engine/Extras/AutoSDK/HostWin64/Linux_x64/v6_clang-3.6.0_ld-2.24_glibc-2.12.2/Unsetup.bat
@@ -0,0 +1,4 @@
+rem TODO: check if this is under UE_SDKS_ROOT ?
+set CurrentFolder=%~dp0
+
+del %CurrentFolder%OutputEnvVars.txt
diff --git a/Engine/Extras/AutoSDK/HostWin64/Linux_x64/v7_clang-3.7.0_ld-2.24_glibc-2.12.2/Setup.bat b/Engine/Extras/AutoSDK/HostWin64/Linux_x64/v7_clang-3.7.0_ld-2.24_glibc-2.12.2/Setup.bat
new file mode 100644
--- /dev/null
+++ b/Engine/Extras/AutoSDK/HostWin64/Linux_x64/v7_clang-3.7.0_ld-2.24_glibc-2.12.2/Setup.bat
@@ -0,0 +1,14 @@
+rem TODO: check if this is under UE_SDKS_ROOT ?
+set CurrentFolder=%~dp0
+echo %CurrentFolder%
+
+REM delete OutputEnvVars from the old location
+del "%CurrentFolder%OutputEnvVars.txt"
+
+set OutputEnvVarsFolder=%~dp0..\
+del "%OutputEnvVarsFolder%OutputEnvVars.txt"
+
+echo LINUX_ROOT=%CurrentFolder%toolchain>"%OutputEnvVarsFolder%OutputEnvVars.txt"
+
+REM support old branches that expect this file in a different location.
+copy "%OutputEnvVarsFolder%OutputEnvVars.txt" "%CurrentFolder%OutputEnvVars.txt"
diff --git a/Engine/Extras/AutoSDK/HostWin64/Linux_x64/v7_clang-3.7.0_ld-2.24_glibc-2.12.2/Unsetup.bat b/Engine/Extras/AutoSDK/HostWin64/Linux_x64/v7_clang-3.7.0_ld-2.24_glibc-2.12.2/Unsetup.bat
new file mode 100644
--- /dev/null
+++ b/Engine/Extras/AutoSDK/HostWin64/Linux_x64/v7_clang-3.7.0_ld-2.24_glibc-2.12.2/Unsetup.bat
@@ -0,0 +1,4 @@
+rem TODO: check if this is under UE_SDKS_ROOT ?
+set CurrentFolder=%~dp0
+
+del %CurrentFolder%OutputEnvVars.txt
diff --git a/Engine/Extras/AutoSDK/HostWin64/Linux_x64/v8_clang-3.9.0-centos7/Setup.bat b/Engine/Extras/AutoSDK/HostWin64/Linux_x64/v8_clang-3.9.0-centos7/Setup.bat
new file mode 100644
--- /dev/null
+++ b/Engine/Extras/AutoSDK/HostWin64/Linux_x64/v8_clang-3.9.0-centos7/Setup.bat
@@ -0,0 +1,18 @@
+rem TODO: check if this is under UE_SDKS_ROOT ?
+set CurrentFolder=%~dp0
+echo %CurrentFolder%
+
+REM delete OutputEnvVars from the old location
+del "%CurrentFolder%OutputEnvVars.txt"
+
+set OutputEnvVarsFolder=%~dp0..\
+del "%OutputEnvVarsFolder%OutputEnvVars.txt"
+
+rem Still set LINUX_ROOT for compatibility
+echo LINUX_ROOT=%CurrentFolder%x86_64-unknown-linux-gnu>"%OutputEnvVarsFolder%OutputEnvVars.txt"
+
+rem set new LINUX_MULTIARCH_ROOT that supersedes it
+echo LINUX_MULTIARCH_ROOT=%CurrentFolder%>>"%OutputEnvVarsFolder%OutputEnvVars.txt"
+
+REM support old branches that expect this file in a different location.
+copy "%OutputEnvVarsFolder%OutputEnvVars.txt" "%CurrentFolder%OutputEnvVars.txt"
diff --git a/Engine/Extras/AutoSDK/HostWin64/Linux_x64/v8_clang-3.9.0-centos7/Unsetup.bat b/Engine/Extras/AutoSDK/HostWin64/Linux_x64/v8_clang-3.9.0-centos7/Unsetup.bat
new file mode 100644
--- /dev/null
+++ b/Engine/Extras/AutoSDK/HostWin64/Linux_x64/v8_clang-3.9.0-centos7/Unsetup.bat
@@ -0,0 +1,4 @@
+rem TODO: check if this is under UE_SDKS_ROOT ?
+set CurrentFolder=%~dp0
+
+del %CurrentFolder%OutputEnvVars.txt
diff --git a/Engine/Extras/AutoSDK/HostWin64/Linux_x64/v9_clang-4.0.0-centos7/Setup.bat b/Engine/Extras/AutoSDK/HostWin64/Linux_x64/v9_clang-4.0.0-centos7/Setup.bat
new file mode 100644
--- /dev/null
+++ b/Engine/Extras/AutoSDK/HostWin64/Linux_x64/v9_clang-4.0.0-centos7/Setup.bat
@@ -0,0 +1,18 @@
+rem TODO: check if this is under UE_SDKS_ROOT ?
+set CurrentFolder=%~dp0
+echo %CurrentFolder%
+
+REM delete OutputEnvVars from the old location
+del "%CurrentFolder%OutputEnvVars.txt"
+
+set OutputEnvVarsFolder=%~dp0..\
+del "%OutputEnvVarsFolder%OutputEnvVars.txt"
+
+rem Still set LINUX_ROOT for compatibility
+echo LINUX_ROOT=%CurrentFolder%x86_64-unknown-linux-gnu>"%OutputEnvVarsFolder%OutputEnvVars.txt"
+
+rem set new LINUX_MULTIARCH_ROOT that supersedes it
+echo LINUX_MULTIARCH_ROOT=%CurrentFolder%>>"%OutputEnvVarsFolder%OutputEnvVars.txt"
+
+REM support old branches that expect this file in a different location.
+copy "%OutputEnvVarsFolder%OutputEnvVars.txt" "%CurrentFolder%OutputEnvVars.txt"
diff --git a/Engine/Extras/AutoSDK/HostWin64/Linux_x64/v9_clang-4.0.0-centos7/Unsetup.bat b/Engine/Extras/AutoSDK/HostWin64/Linux_x64/v9_clang-4.0.0-centos7/Unsetup.bat
new file mode 100644
--- /dev/null
+++ b/Engine/Extras/AutoSDK/HostWin64/Linux_x64/v9_clang-4.0.0-centos7/Unsetup.bat
@@ -0,0 +1,4 @@
+rem TODO: check if this is under UE_SDKS_ROOT ?
+set CurrentFolder=%~dp0
+
+del %CurrentFolder%OutputEnvVars.txt
diff --git a/Engine/Extras/AutoSDK/HostWin64/Linux_x64/vNext_clang-Next-centos7/Setup.bat b/Engine/Extras/AutoSDK/HostWin64/Linux_x64/vNext_clang-Next-centos7/Setup.bat
new file mode 100644
--- /dev/null
+++ b/Engine/Extras/AutoSDK/HostWin64/Linux_x64/vNext_clang-Next-centos7/Setup.bat
@@ -0,0 +1,18 @@
+rem TODO: check if this is under UE_SDKS_ROOT ?
+set CurrentFolder=%~dp0
+echo %CurrentFolder%
+
+REM delete OutputEnvVars from the old location
+del "%CurrentFolder%OutputEnvVars.txt"
+
+set OutputEnvVarsFolder=%~dp0..\
+del "%OutputEnvVarsFolder%OutputEnvVars.txt"
+
+rem Still set LINUX_ROOT for compatibility
+echo LINUX_ROOT=%CurrentFolder%x86_64-unknown-linux-gnu>"%OutputEnvVarsFolder%OutputEnvVars.txt"
+
+rem set new LINUX_MULTIARCH_ROOT that supersedes it
+echo LINUX_MULTIARCH_ROOT=%CurrentFolder%>>"%OutputEnvVarsFolder%OutputEnvVars.txt"
+
+REM support old branches that expect this file in a different location.
+copy "%OutputEnvVarsFolder%OutputEnvVars.txt" "%CurrentFolder%OutputEnvVars.txt"
diff --git a/Engine/Extras/AutoSDK/HostWin64/Linux_x64/vNext_clang-Next-centos7/Unsetup.bat b/Engine/Extras/AutoSDK/HostWin64/Linux_x64/vNext_clang-Next-centos7/Unsetup.bat
new file mode 100644
--- /dev/null
+++ b/Engine/Extras/AutoSDK/HostWin64/Linux_x64/vNext_clang-Next-centos7/Unsetup.bat
@@ -0,0 +1,4 @@
+rem TODO: check if this is under UE_SDKS_ROOT ?
+set CurrentFolder=%~dp0
+
+del %CurrentFolder%OutputEnvVars.txt
