Add support for using MSBuild from AutoSDK directories.

diff --git a/Engine/Build/BatchFiles/GetMSBuildPath.bat b/Engine/Build/BatchFiles/GetMSBuildPath.bat
--- a/Engine/Build/BatchFiles/GetMSBuildPath.bat
+++ b/Engine/Build/BatchFiles/GetMSBuildPath.bat
@@ -8,6 +8,13 @@ rem ## The discovered path is set to the MSBUILD_EXE environment variable on suc
 
 set MSBUILD_EXE=
 
+rem ## Check to see if we've got a copy of MSBuild in AutoSDKs
+if "%UE_SDKS_ROOT%"=="" goto no_autosdk
+if not exist "%UE_SDKS_ROOT%\HostWin64\Win64\MSBuild\VS2019\MSBuild\Current\Bin\MSBuild.exe" goto no_autosdk
+set MSBUILD_EXE="%UE_SDKS_ROOT%\HostWin64\Win64\MSBuild\VS2019\MSBuild\Current\Bin\MSBuild.exe"
+goto Succeeded
+:no_autosdk
+
 rem ## Try to get the MSBuild 15 path using vswhere (see [REDACTED_PROFILE_URL]/vswhere). VS2019 preview puts the executable in a "Current" subfolder.
 if not exist "%ProgramFiles(x86)%\Microsoft Visual Studio\Installer\vswhere.exe" goto no_vswhere
 for /f "delims=" %%i in ('"%ProgramFiles(x86)%\Microsoft Visual Studio\Installer\vswhere" -prerelease -latest -products * -requires Microsoft.Component.MSBuild -property installationPath') do (
diff --git a/Engine/Extras/AutoSDK/HostWin64/Win64/README.md b/Engine/Extras/AutoSDK/HostWin64/Win64/README.md
--- a/Engine/Extras/AutoSDK/HostWin64/Win64/README.md
+++ b/Engine/Extras/AutoSDK/HostWin64/Win64/README.md
@@ -30,6 +30,22 @@ The NETFXSDK directory may also be distributed in the same manner.
 
     [AutoSDK Root]\HostWin64\Win64\Windows Kits\NETFXSDK\...
 
+## MSBuild
+
+A single MSBuild installation is supported through AutoSDK, following the pattern Microsoft have introduced with VS2019. The folders:
+
+    C:\Program Files (x86)\Microsoft Visual Studio\2019\Professional\MSBuild\...
+    C:\Program Files (x86)\Microsoft Visual Studio\2019\Professional\Common7\IDE\CommonExtensions\Microsoft\NuGet\...
+    C:\Program Files (x86)\Microsoft Visual Studio\2019\Professional\Team Tools\Static Analysis Tools\...
+
+Should be copied to:
+
+    [AutoSDK Root]\HostWin64\Win64\MSBuild\VS2019\MSBuild\...
+    [AutoSDK Root]\HostWin64\Win64\MSBuild\VS2019\Common7\IDE\CommonExtensions\Microsoft\NuGet\...
+	[AutoSDK Root]\HostWin64\Win64\MSBuild\VS2019\Team Tools\Static Analysis Tools\...
+
+The GetMsBuildPath.bat batch file and associated scripts (RunUAT.bat, GenerateProjectFiles.bat) will also check for MSBuild.exe in this location.
+
 ## DIA SDK
 
 The DIA SDK is distributed with Visual Studio, and is typically installed to the DIA SDK folder under the Visual Studio installation directory. For example:
