Merging AutoSDK changes.

#rnx

diff --git a/Engine/Extras/AutoSDK/HostLinux/Linux_x64/v15_clang-8.0.1-centos7/setup.sh b/Engine/Extras/AutoSDK/HostLinux/Linux_x64/v15_clang-8.0.1-centos7/setup.sh
new file mode 100755
--- /dev/null
+++ b/Engine/Extras/AutoSDK/HostLinux/Linux_x64/v15_clang-8.0.1-centos7/setup.sh
@@ -0,0 +1,14 @@
+#!/bin/bash
+
+# TODO: check if this is under UE_SDKS_ROOT ?
+CURRENT_FOLDER=$(cd "$(dirname "$BASH_SOURCE")" ; pwd)
+echo "$CURRENT_FOLDER"
+
+OUTPUT_ENV_VARS_FOLDER="$CURRENT_FOLDER/.."
+rm "$OUTPUT_ENV_VARS_FOLDER/OutputEnvVars.txt" 2> /dev/null
+
+# set new LINUX_MULTIARCH_ROOT that supersedes it
+echo LINUX_MULTIARCH_ROOT="$CURRENT_FOLDER" > "$OUTPUT_ENV_VARS_FOLDER/OutputEnvVars.txt"
+
+# support old branches that expect this file in a different location.
+cp "$OUTPUT_ENV_VARS_FOLDER/OutputEnvVars.txt" "$CURRENT_FOLDER"
diff --git a/Engine/Extras/AutoSDK/HostLinux/Linux_x64/v15_clang-8.0.1-centos7/unsetup.sh b/Engine/Extras/AutoSDK/HostLinux/Linux_x64/v15_clang-8.0.1-centos7/unsetup.sh
new file mode 100755
--- /dev/null
+++ b/Engine/Extras/AutoSDK/HostLinux/Linux_x64/v15_clang-8.0.1-centos7/unsetup.sh
@@ -0,0 +1,9 @@
+#!/bin/bash
+
+# TODO: check if this is under UE_SDKS_ROOT ?
+CURRENT_FOLDER=$(cd "$(dirname "$BASH_SOURCE")" ; pwd)
+echo "$CURRENT_FOLDER"
+
+OUTPUT_ENV_VARS_FOLDER="$CURRENT_FOLDER/.."
+
+rm "$OUTPUT_ENV_VARS_FOLDER/OutputEnvVars.txt"
diff --git a/Engine/Extras/AutoSDK/HostMac/Lumin/0.22/setup.sh b/Engine/Extras/AutoSDK/HostMac/Lumin/0.22/setup.sh
new file mode 100755
--- /dev/null
+++ b/Engine/Extras/AutoSDK/HostMac/Lumin/0.22/setup.sh
@@ -0,0 +1,5 @@
+#!/bin/bash
+
+cd `dirname $0`
+echo MLSDK=`pwd` >../OutputEnvVars.txt
+
diff --git a/Engine/Extras/AutoSDK/HostWin64/Linux_x64/v15_clang-8.0.1-centos7/Setup.bat b/Engine/Extras/AutoSDK/HostWin64/Linux_x64/v15_clang-8.0.1-centos7/Setup.bat
new file mode 100644
--- /dev/null
+++ b/Engine/Extras/AutoSDK/HostWin64/Linux_x64/v15_clang-8.0.1-centos7/Setup.bat
@@ -0,0 +1,18 @@
+rem TODO: check if this is under UE_SDKS_ROOT ?
+set CurrentFolder=%~dp0
+echo %CurrentFolder%
+
+REM delete OutputEnvVars from the old location
+del "%CurrentFolder%OutputEnvVars.txt"
+
+set OutputEnvVarsFolder=%~dp0..\
+del "%OutputEnvVarsFolder%OutputEnvVars.txt"
+
+rem Still set LINUX_ROOT for compatibility
+echo LINUX_ROOT=%CurrentFolder%x86_64-unknown-linux-gnu>"%OutputEnvVarsFolder%OutputEnvVars.txt"
+
+rem set new LINUX_MULTIARCH_ROOT that supersedes it
+echo LINUX_MULTIARCH_ROOT=%CurrentFolder%>>"%OutputEnvVarsFolder%OutputEnvVars.txt"
+
+REM support old branches that expect this file in a different location.
+copy "%OutputEnvVarsFolder%OutputEnvVars.txt" "%CurrentFolder%OutputEnvVars.txt"
diff --git a/Engine/Extras/AutoSDK/HostWin64/Linux_x64/v15_clang-8.0.1-centos7/Unsetup.bat b/Engine/Extras/AutoSDK/HostWin64/Linux_x64/v15_clang-8.0.1-centos7/Unsetup.bat
new file mode 100644
--- /dev/null
+++ b/Engine/Extras/AutoSDK/HostWin64/Linux_x64/v15_clang-8.0.1-centos7/Unsetup.bat
@@ -0,0 +1,4 @@
+rem TODO: check if this is under UE_SDKS_ROOT ?
+set CurrentFolder=%~dp0
+
+del %CurrentFolder%OutputEnvVars.txt
diff --git a/Engine/Extras/AutoSDK/HostWin64/Lumin/0.22/Setup.bat b/Engine/Extras/AutoSDK/HostWin64/Lumin/0.22/Setup.bat
new file mode 100644
--- /dev/null
+++ b/Engine/Extras/AutoSDK/HostWin64/Lumin/0.22/Setup.bat
@@ -0,0 +1,13 @@
+set CurrentFolder=%~dp0
+echo %CurrentFolder%
+
+REM delete OutputEnvVars from the old location
+del "%CurrentFolder%OutputEnvVars.txt"
+
+set OutputEnvVarsFolder=%~dp0..\
+del "%OutputEnvVarsFolder%OutputEnvVars.txt"
+
+
+
+REM set android sdk environment variables
+echo MLSDK=%CurrentFolder%>>"%OutputEnvVarsFolder%OutputEnvVars.txt"
