Updating AutoSDKs

#rnx

diff --git a/Engine/Extras/AutoSDK/HostMac/Lumin/0.22/unsetup.sh b/Engine/Extras/AutoSDK/HostMac/Lumin/0.22/unsetup.sh
new file mode 100755
--- /dev/null
+++ b/Engine/Extras/AutoSDK/HostMac/Lumin/0.22/unsetup.sh
@@ -0,0 +1,4 @@
+#!/bin/bash
+
+cd `dirname $0`
+rm ../OutputEnvVars.txt 2>/dev/null || true
diff --git a/Engine/Extras/AutoSDK/HostWin64/Lumin/0.22/Unsetup.bat b/Engine/Extras/AutoSDK/HostWin64/Lumin/0.22/Unsetup.bat
new file mode 100644
--- /dev/null
+++ b/Engine/Extras/AutoSDK/HostWin64/Lumin/0.22/Unsetup.bat
@@ -0,0 +1,4 @@
+rem TODO: check if this is under UE_SDKS_ROOT ?
+set CurrentFolder=%~dp0
+
+del %CurrentFolder%OutputEnvVars.txt
