- Moved some Engine/Extras/AutoSDKs/ platforms to PlatformExtensions

diff --git a/Engine/Extras/AutoSDK/README.md b/Engine/Extras/AutoSDK/README.md
--- a/Engine/Extras/AutoSDK/README.md
+++ b/Engine/Extras/AutoSDK/README.md
@@ -36,4 +36,4 @@ Setup scripts written by Epic for supported SDK versions are included under this
 
 ## Platforms
 
-More information about adding SDKs for each platform is given by README.md files in HostPlatform/TargetPlatform subfolders.
\ No newline at end of file
+More information about adding SDKs for each platform is given by README.md files in HostPlatform/TargetPlatform subfolders. For PlatformExtension platforms, you will need to merge AutoSDKs files from Engine/Platforms/<Platform>/Extras/AutoSDK into your UE_SDKS_ROOT
\ No newline at end of file
