Fixing incorrect JAVA_HOME setup

#android

diff --git a/Engine/Extras/AutoSDK/HostWin64/Android/-26/setup.bat b/Engine/Extras/AutoSDK/HostWin64/Android/-26/setup.bat
new file mode 100755
--- /dev/null
+++ b/Engine/Extras/AutoSDK/HostWin64/Android/-26/setup.bat
@@ -0,0 +1,32 @@
+rem TODO: check if this is under UE_SDKS_ROOT ?
+set CurrentFolder=%~dp0
+echo %CurrentFolder%
+
+REM delete OutputEnvVars from the old location
+del "%CurrentFolder%OutputEnvVars.txt"
+
+set OutputEnvVarsFolder=%~dp0..\
+del "%OutputEnvVarsFolder%OutputEnvVars.txt"
+
+
+
+REM set android sdk environment variables
+echo ANDROID_HOME=%CurrentFolder%>>"%OutputEnvVarsFolder%OutputEnvVars.txt"
+
+REM set ndk environment variables
+echo ANDROID_NDK_ROOT=%CurrentFolder%ndk\25.1.8937393>>"%OutputEnvVarsFolder%OutputEnvVars.txt"
+echo NDK_ROOT=%CurrentFolder%ndk\25.1.8937393>>"%OutputEnvVarsFolder%OutputEnvVars.txt"
+echo NDKROOT=%CurrentFolder%ndk\25.1.8937393>>"%OutputEnvVarsFolder%OutputEnvVars.txt"
+
+REM set JDK directory
+echo JAVA_HOME=%CurrentFolder%jbr>>"%OutputEnvVarsFolder%OutputEnvVars.txt"
+
+REM Set the Andriod SWT path 
+echo ANDROID_SWT=%CurrentFolder%\tools\lib\x86_64>>"%OutputEnvVarsFolder%OutputEnvVars.txt"
+
+
+
+
+REM support old branches that expect this file in a different location.
+copy "%OutputEnvVarsFolder%OutputEnvVars.txt" "%CurrentFolder%OutputEnvVars.txt"
+
diff --git a/Engine/Extras/AutoSDK/HostWin64/Android/-26/unsetup.bat b/Engine/Extras/AutoSDK/HostWin64/Android/-26/unsetup.bat
new file mode 100755
--- /dev/null
+++ b/Engine/Extras/AutoSDK/HostWin64/Android/-26/unsetup.bat
@@ -0,0 +1,4 @@
+rem TODO: check if this is under UE_SDKS_ROOT ?
+set CurrentFolder=%~dp0
+
+del %CurrentFolder%OutputEnvVars.txt
